[AI Toolsets]: Move ConversationToolset under the Toolsets directory

#Synced-CL 52368392

diff --git a/Engine/Plugins/Experimental/ConversationToolset/ConversationToolset.uplugin b/Engine/Plugins/Experimental/ConversationToolset/ConversationToolset.uplugin
deleted file mode 100644
diff --git a/Engine/Plugins/Experimental/ConversationToolset/Source/ConversationToolset/ConversationToolset.Build.cs b/Engine/Plugins/Experimental/ConversationToolset/Source/ConversationToolset/ConversationToolset.Build.cs
deleted file mode 100644
diff --git a/Engine/Plugins/Experimental/ConversationToolset/Source/ConversationToolset/Private/ConversationToolset.cpp b/Engine/Plugins/Experimental/ConversationToolset/Source/ConversationToolset/Private/ConversationToolset.cpp
deleted file mode 100644
diff --git a/Engine/Plugins/Experimental/ConversationToolset/Source/ConversationToolset/Public/ConversationToolset.h b/Engine/Plugins/Experimental/ConversationToolset/Source/ConversationToolset/Public/ConversationToolset.h
deleted file mode 100644
diff --git a/Engine/Plugins/Experimental/Toolsets/ConversationToolset/ConversationToolset.uplugin b/Engine/Plugins/Experimental/Toolsets/ConversationToolset/ConversationToolset.uplugin
new file mode 100644
--- /dev/null
+++ b/Engine/Plugins/Experimental/Toolsets/ConversationToolset/ConversationToolset.uplugin
@@ -0,0 +1,38 @@
+{
+	"FileVersion": 3,
+	"Version": 1,
+	"VersionName": "1.0",
+	"FriendlyName": "ConversationToolset",
+	"Description": "Toolset for Conversation Systems",
+	"Category": "Other",
+	"CreatedBy": "Epic Games",
+	"CreatedByURL": "",
+	"DocsURL": "",
+	"MarketplaceURL": "",
+	"SupportURL": "",
+	"EditorOnly": true,
+	"CanContainContent": true,
+	"IsBetaVersion": false,
+	"IsExperimentalVersion": true,
+	"NoRedist": true,
+	"Installed": false,
+	"EnabledByDefault": false,
+	"Modules": [
+		{
+			"Name": "ConversationToolset",
+			"Type": "Editor",
+			"LoadingPhase": "Default",
+			"TargetAllowList": ["Editor"]
+		}
+	],
+	"Plugins": [
+		{
+			"Name": "CommonConversation",
+			"Enabled": true
+		},
+		{
+			"Name": "ToolsetRegistry",
+			"Enabled": true
+		}
+	]
+}
diff --git a/Engine/Plugins/Experimental/Toolsets/ConversationToolset/Source/ConversationToolset/ConversationToolset.Build.cs b/Engine/Plugins/Experimental/Toolsets/ConversationToolset/Source/ConversationToolset/ConversationToolset.Build.cs
new file mode 100644
--- /dev/null
+++ b/Engine/Plugins/Experimental/Toolsets/ConversationToolset/Source/ConversationToolset/ConversationToolset.Build.cs
@@ -0,0 +1,12 @@
+// Copyright Epic Games, Inc. All Rights Reserved.
+
+using UnrealBuildTool;
+
+public class ConversationToolset : ModuleRules
+{
+	public ConversationToolset(ReadOnlyTargetRules Target) : base(Target)
+	{
+		PCHUsage = PCHUsageMode.UseExplicitOrSharedPCHs;
+		PublicDependencyModuleNames.Add("Core");
+	}
+}
diff --git a/Engine/Plugins/Experimental/Toolsets/ConversationToolset/Source/ConversationToolset/Private/ConversationToolset.cpp b/Engine/Plugins/Experimental/Toolsets/ConversationToolset/Source/ConversationToolset/Private/ConversationToolset.cpp
new file mode 100644
--- /dev/null
+++ b/Engine/Plugins/Experimental/Toolsets/ConversationToolset/Source/ConversationToolset/Private/ConversationToolset.cpp
@@ -0,0 +1,16 @@
+// Copyright Epic Games, Inc. All Rights Reserved.
+
+#include "ConversationToolset.h"
+
+#include "Modules/ModuleManager.h"
+
+/*virtual*/ void FConversationToolsetModule::StartupModule()
+{
+}
+
+/*virtual*/ void FConversationToolsetModule::ShutdownModule()
+{
+}
+
+IMPLEMENT_MODULE(FConversationToolsetModule, ConversationToolset);
+DEFINE_LOG_CATEGORY(LogConversationToolset);
diff --git a/Engine/Plugins/Experimental/Toolsets/ConversationToolset/Source/ConversationToolset/Public/ConversationToolset.h b/Engine/Plugins/Experimental/Toolsets/ConversationToolset/Source/ConversationToolset/Public/ConversationToolset.h
new file mode 100644
--- /dev/null
+++ b/Engine/Plugins/Experimental/Toolsets/ConversationToolset/Source/ConversationToolset/Public/ConversationToolset.h
@@ -0,0 +1,15 @@
+// Copyright Epic Games, Inc. All Rights Reserved.
+
+#pragma once
+
+#include "Modules/ModuleInterface.h"
+#include "Logging/LogMacros.h"
+
+DECLARE_LOG_CATEGORY_EXTERN(LogConversationToolset, Log, Log);
+
+class FConversationToolsetModule : public IModuleInterface
+{
+public:
+	virtual void StartupModule() override;
+	virtual void ShutdownModule() override;
+};
